Matching LOCTEXT_NAMESPACE #def/#undef

- Also changed some #undef's to be more uniform

diff --git a/Engine/Plugins/Messaging/MessagingDebugger/Source/MessagingDebugger/Private/Widgets/MessageData/SMessagingMessageData.h b/Engine/Plugins/Messaging/MessagingDebugger/Source/MessagingDebugger/Private/Widgets/MessageData/SMessagingMessageData.h
--- a/Engine/Plugins/Messaging/MessagingDebugger/Source/MessagingDebugger/Private/Widgets/MessageData/SMessagingMessageData.h
+++ b/Engine/Plugins/Messaging/MessagingDebugger/Source/MessagingDebugger/Private/Widgets/MessageData/SMessagingMessageData.h
@@ -62,6 +62,3 @@ private:
 	/** Holds the details text box. */
 	TSharedPtr<SMultiLineEditableTextBox> TextBox;
 };
-
-
-#undef LOCTEXT_NAMESPACE
diff --git a/Engine/Plugins/NetcodeUnitTest/NetcodeUnitTest/Source/NetcodeUnitTest/Private/UI/SLogDialog.cpp b/Engine/Plugins/NetcodeUnitTest/NetcodeUnitTest/Source/NetcodeUnitTest/Private/UI/SLogDialog.cpp
--- a/Engine/Plugins/NetcodeUnitTest/NetcodeUnitTest/Source/NetcodeUnitTest/Private/UI/SLogDialog.cpp
+++ b/Engine/Plugins/NetcodeUnitTest/NetcodeUnitTest/Source/NetcodeUnitTest/Private/UI/SLogDialog.cpp
@@ -279,3 +279,5 @@ TSharedRef<SWindow> OpenLogDialog_NonModal(EAppMsgType::Type InMessageType, cons
 
 	return MsgWindow.ToSharedRef();
 }
+
+#undef LOCTEXT_NAMESPACE
diff --git a/Engine/Plugins/Runtime/Analytics/AnalyticsMulticast/Source/AnaltyicsMulticastEditor/Private/AnalyticsMulticastEditor.cpp b/Engine/Plugins/Runtime/Analytics/AnalyticsMulticast/Source/AnaltyicsMulticastEditor/Private/AnalyticsMulticastEditor.cpp
--- a/Engine/Plugins/Runtime/Analytics/AnalyticsMulticast/Source/AnaltyicsMulticastEditor/Private/AnalyticsMulticastEditor.cpp
+++ b/Engine/Plugins/Runtime/Analytics/AnalyticsMulticast/Source/AnaltyicsMulticastEditor/Private/AnalyticsMulticastEditor.cpp
@@ -86,3 +86,5 @@ FString UAnalyticsMulticastSettings::BuildStringFromArray(TArray<FString>& Array
 	}
 	return List;
 }
+
+#undef LOCTEXT_NAMESPACE
diff --git a/Engine/Plugins/Runtime/Analytics/Apsalar/Source/ApsalarEditor/Private/ApsalarEditor.cpp b/Engine/Plugins/Runtime/Analytics/Apsalar/Source/ApsalarEditor/Private/ApsalarEditor.cpp
--- a/Engine/Plugins/Runtime/Analytics/Apsalar/Source/ApsalarEditor/Private/ApsalarEditor.cpp
+++ b/Engine/Plugins/Runtime/Analytics/Apsalar/Source/ApsalarEditor/Private/ApsalarEditor.cpp
@@ -97,3 +97,5 @@ void UApsalarSettings::WriteConfigStruct(const FString& Section, FApsalarAnalyti
 	FAnalytics::Get().WriteConfigValueToIni(GetIniName(), Section, TEXT("MaxBufferSize"), FString::Printf(TEXT("%d"), Source.MaxBufferSize));
 	FAnalytics::Get().WriteConfigValueToIni(GetIniName(), Section, TEXT("ManuallyReportRevenue"), Source.ManuallyReportRevenue ? TEXT("true") : TEXT("false"));
 }
+
+#undef LOCTEXT_NAMESPACE
diff --git a/Engine/Plugins/Runtime/Analytics/Flurry/Source/FlurryEditor/Private/FlurryEditor.cpp b/Engine/Plugins/Runtime/Analytics/Flurry/Source/FlurryEditor/Private/FlurryEditor.cpp
--- a/Engine/Plugins/Runtime/Analytics/Flurry/Source/FlurryEditor/Private/FlurryEditor.cpp
+++ b/Engine/Plugins/Runtime/Analytics/Flurry/Source/FlurryEditor/Private/FlurryEditor.cpp
@@ -58,3 +58,5 @@ void UFlurrySettings::WriteConfigSettings()
 	FAnalytics::Get().WriteConfigValueToIni(GetIniName(), GetDebugIniSection(), TEXT("FlurryApiKey"), DebugApiKey);
 	FAnalytics::Get().WriteConfigValueToIni(GetIniName(), GetDevelopmentIniSection(), TEXT("FlurryApiKey"), DevelopmentApiKey);
 }
+
+#undef LOCTEXT_NAMESPACE
diff --git a/Engine/Source/Developer/AssetTools/Private/AssetTypeActions/AssetTypeActions_DestructibleMesh.cpp b/Engine/Source/Developer/AssetTools/Private/AssetTypeActions/AssetTypeActions_DestructibleMesh.cpp
--- a/Engine/Source/Developer/AssetTools/Private/AssetTypeActions/AssetTypeActions_DestructibleMesh.cpp
+++ b/Engine/Source/Developer/AssetTools/Private/AssetTypeActions/AssetTypeActions_DestructibleMesh.cpp
@@ -25,4 +25,6 @@ void FAssetTypeActions_DestructibleMesh::OpenAssetEditor( const TArray<UObject*>
 			TSharedRef< IDestructibleMeshEditor > NewDestructibleMeshEditor = DestructibleMeshEditorModule.CreateDestructibleMeshEditor( EToolkitMode::Standalone, EditWithinLevelEditor, Mesh );
 		}
 	}
-}
\ No newline at end of file
+}
+
+#undef LOCTEXT_NAMESPACE
diff --git a/Engine/Source/Developer/FriendsAndChat/Private/Models/FriendRecentPlayerItem.cpp b/Engine/Source/Developer/FriendsAndChat/Private/Models/FriendRecentPlayerItem.cpp
--- a/Engine/Source/Developer/FriendsAndChat/Private/Models/FriendRecentPlayerItem.cpp
+++ b/Engine/Source/Developer/FriendsAndChat/Private/Models/FriendRecentPlayerItem.cpp
@@ -156,5 +156,3 @@ EInviteStatus::Type FFriendRecentPlayerItem::GetInviteStatus()
 {
 	return EInviteStatus::Unknown;
 }
-
-#undef LOCTEXT_NAMESPACE
diff --git a/Engine/Source/Developer/TreeMap/TreeMapModule.cpp b/Engine/Source/Developer/TreeMap/TreeMapModule.cpp
--- a/Engine/Source/Developer/TreeMap/TreeMapModule.cpp
+++ b/Engine/Source/Developer/TreeMap/TreeMapModule.cpp
@@ -25,6 +25,3 @@ public:
 
 
 IMPLEMENT_MODULE(FTreeMapModule, TreeMap);
-
-
-#undef LOCTEXT_NAMESPACE
diff --git a/Engine/Source/Editor/DataTableEditor/Private/SRowEditor.cpp b/Engine/Source/Editor/DataTableEditor/Private/SRowEditor.cpp
--- a/Engine/Source/Editor/DataTableEditor/Private/SRowEditor.cpp
+++ b/Engine/Source/Editor/DataTableEditor/Private/SRowEditor.cpp
@@ -510,3 +510,4 @@ void SRowEditor::Construct(const FArguments& InArgs, UDataTable* Changed)
 	];
 }
 
+#undef LOCTEXT_NAMESPACE
diff --git a/Engine/Source/Editor/DetailCustomizations/Public/Customizations/MobilityCustomization.h b/Engine/Source/Editor/DetailCustomizations/Public/Customizations/MobilityCustomization.h
--- a/Engine/Source/Editor/DetailCustomizations/Public/Customizations/MobilityCustomization.h
+++ b/Engine/Source/Editor/DetailCustomizations/Public/Customizations/MobilityCustomization.h
@@ -29,6 +29,3 @@ private:
 	
 	TSharedPtr<IPropertyHandle> MobilityHandle;
 };
-
-#undef LOCTEXT_NAMESPACE
-
diff --git a/Engine/Source/Editor/EditorWidgets/Private/SDropTarget.cpp b/Engine/Source/Editor/EditorWidgets/Private/SDropTarget.cpp
--- a/Engine/Source/Editor/EditorWidgets/Private/SDropTarget.cpp
+++ b/Engine/Source/Editor/EditorWidgets/Private/SDropTarget.cpp
@@ -202,3 +202,5 @@ int32 SDropTarget::OnPaint(const FPaintArgs& Args, const FGeometry& AllottedGeom
 
 	return LayerId;
 }
+
+#undef LOCTEXT_NAMESPACE
diff --git a/Engine/Source/Editor/GameProjectGeneration/Private/SGetSuggestedIDEWidget.cpp b/Engine/Source/Editor/GameProjectGeneration/Private/SGetSuggestedIDEWidget.cpp
--- a/Engine/Source/Editor/GameProjectGeneration/Private/SGetSuggestedIDEWidget.cpp
+++ b/Engine/Source/Editor/GameProjectGeneration/Private/SGetSuggestedIDEWidget.cpp
@@ -81,4 +81,6 @@ void SGetSuggestedIDEWidget::OnIDEInstallerDownloadComplete(bool bWasSuccessful)
 
 		IDEDownloadNotification->ExpireAndFadeout();
 	}
-}
\ No newline at end of file
+}
+
+#undef LOCTEXT_NAMESPACE
diff --git a/Engine/Source/Editor/PropertyEditor/Private/UserInterface/PropertyEditor/SPropertyEditorText.cpp b/Engine/Source/Editor/PropertyEditor/Private/UserInterface/PropertyEditor/SPropertyEditorText.cpp
--- a/Engine/Source/Editor/PropertyEditor/Private/UserInterface/PropertyEditor/SPropertyEditorText.cpp
+++ b/Engine/Source/Editor/PropertyEditor/Private/UserInterface/PropertyEditor/SPropertyEditorText.cpp
@@ -182,3 +182,5 @@ bool SPropertyEditorText::IsReadOnly() const
 {
 	return !CanEdit();
 }
+
+#undef LOCTEXT_NAMESPACE
diff --git a/Engine/Source/Editor/UnrealEd/Private/Commandlets/ChunkManifestGenerator.cpp b/Engine/Source/Editor/UnrealEd/Private/Commandlets/ChunkManifestGenerator.cpp
--- a/Engine/Source/Editor/UnrealEd/Private/Commandlets/ChunkManifestGenerator.cpp
+++ b/Engine/Source/Editor/UnrealEd/Private/Commandlets/ChunkManifestGenerator.cpp
@@ -1249,3 +1249,5 @@ FString FChunkManifestGenerator::GetShortestReferenceChain(FName PackageName, in
 
 	return StringChain;
 }
+
+#undef LOCTEXT_NAMESPACE
diff --git a/Engine/Source/Programs/CrashReporter/CrashReportClient/Private/Mac/MacErrorReport.cpp b/Engine/Source/Programs/CrashReporter/CrashReportClient/Private/Mac/MacErrorReport.cpp
--- a/Engine/Source/Programs/CrashReporter/CrashReportClient/Private/Mac/MacErrorReport.cpp
+++ b/Engine/Source/Programs/CrashReporter/CrashReportClient/Private/Mac/MacErrorReport.cpp
@@ -135,3 +135,5 @@ FText FMacErrorReport::DiagnoseReport() const
 		}
 	}
 }
+
+#undef LOCTEXT_NAMESPACE
diff --git a/Engine/Source/Runtime/Analytics/AnalyticsVisualEditing/Private/AnalyticsVisualEditing.cpp b/Engine/Source/Runtime/Analytics/AnalyticsVisualEditing/Private/AnalyticsVisualEditing.cpp
--- a/Engine/Source/Runtime/Analytics/AnalyticsVisualEditing/Private/AnalyticsVisualEditing.cpp
+++ b/Engine/Source/Runtime/Analytics/AnalyticsVisualEditing/Private/AnalyticsVisualEditing.cpp
@@ -55,3 +55,5 @@ void UAnalyticsSettings::WriteConfigSettings()
 	FAnalytics::Get().WriteConfigValueToIni(GetIniName(), GetTestIniSection(), TEXT("ProviderModuleName"), TestProviderName);
 	FAnalytics::Get().WriteConfigValueToIni(GetIniName(), GetDebugIniSection(), TEXT("ProviderModuleName"), DebugProviderName);
 }
+
+#undef LOCTEXT_NAMESPACE
diff --git a/Engine/Source/Runtime/AppFramework/Private/Widgets/Colors/SEyeDropperButton.h b/Engine/Source/Runtime/AppFramework/Private/Widgets/Colors/SEyeDropperButton.h
--- a/Engine/Source/Runtime/AppFramework/Private/Widgets/Colors/SEyeDropperButton.h
+++ b/Engine/Source/Runtime/AppFramework/Private/Widgets/Colors/SEyeDropperButton.h
@@ -77,6 +77,3 @@ private:
 	bool bWasLeft;
 	bool bWasReEntered;
 };
-
-
-#undef LOCTEXT_NAMESPACE
diff --git a/Engine/Source/Runtime/Engine/Private/AI/Navigation/NavigationSystem.cpp b/Engine/Source/Runtime/Engine/Private/AI/Navigation/NavigationSystem.cpp
--- a/Engine/Source/Runtime/Engine/Private/AI/Navigation/NavigationSystem.cpp
+++ b/Engine/Source/Runtime/Engine/Private/AI/Navigation/NavigationSystem.cpp
@@ -4044,4 +4044,6 @@ FVector UNavigationSystem::GetRandomPointInRadius(UObject* WorldContextObject, c
 bool UNavigationSystem::GetRandomPointInRadius(const FVector& Origin, float Radius, FNavLocation& ResultLocation, ANavigationData* NavData, TSharedPtr<const FNavigationQueryFilter> QueryFilter) const
 {
 	return GetRandomReachablePointInRadius(Origin, Radius, ResultLocation, NavData, QueryFilter);
-}
\ No newline at end of file
+}
+
+#undef LOCTEXT_NAMESPACE
