fix for IOS compile issue with Analytics

diff --git a/Engine/Plugins/Runtime/Analytics/Apsalar/Source/IOSApsalar/Private/IOSApsalar.cpp b/Engine/Plugins/Runtime/Analytics/Apsalar/Source/IOSApsalar/Private/IOSApsalar.cpp
--- a/Engine/Plugins/Runtime/Analytics/Apsalar/Source/IOSApsalar/Private/IOSApsalar.cpp
+++ b/Engine/Plugins/Runtime/Analytics/Apsalar/Source/IOSApsalar/Private/IOSApsalar.cpp
@@ -228,7 +228,7 @@ void FAnalyticsProviderApsalar::RecordEvent(const FString& EventName, const TArr
 			for	(auto Attr : Attributes)
 			{
 				NSString* AttrName = [NSString stringWithFString : Attr.AttrName];
-				NSString* AttrValue = [NSString stringWithFString : Attr.AttrValue];
+				NSString* AttrValue = [NSString stringWithFString : Attr.ToString()];
 				[AttributesDict setValue:AttrValue forKey:AttrName];
 			}
 			[Apsalar event:ConvertedEventName withArgs:AttributesDict];
