- Moved Engine/Source/ThirdParty/.../NotForLIcensees (and NoRedist) to Engine/Restricted

diff --git a/Engine/Plugins/Online/OnlineSubsystemFacebook/Source/OnlineSubsystemFacebook.Build.cs b/Engine/Plugins/Online/OnlineSubsystemFacebook/Source/OnlineSubsystemFacebook.Build.cs
--- a/Engine/Plugins/Online/OnlineSubsystemFacebook/Source/OnlineSubsystemFacebook.Build.cs
+++ b/Engine/Plugins/Online/OnlineSubsystemFacebook/Source/OnlineSubsystemFacebook.Build.cs
@@ -38,7 +38,7 @@ public class OnlineSubsystemFacebook : ModuleRules
 			string FacebookNFLDir = "";
 			try
 			{
-				FacebookNFLDir = System.IO.Path.Combine(ModuleDirectory, "ThirdParty", "Android", "NotForLicensees", "FacebookSDK");
+				FacebookNFLDir = System.IO.Path.Combine(EngineDirectory, "Restricted/NotForLicensees/Plugins/Online/OnlineSubsystemFacebook/Source/ThirdParty/Android/FacebookSDK");
 				bHasFacebookSDK = System.IO.Directory.Exists(FacebookNFLDir);
 			}
 			catch (System.Exception)
diff --git a/Engine/Plugins/Online/OnlineSubsystemTencent/Source/OnlineSubsystemTencent.Build.cs b/Engine/Plugins/Online/OnlineSubsystemTencent/Source/OnlineSubsystemTencent.Build.cs
--- a/Engine/Plugins/Online/OnlineSubsystemTencent/Source/OnlineSubsystemTencent.Build.cs
+++ b/Engine/Plugins/Online/OnlineSubsystemTencent/Source/OnlineSubsystemTencent.Build.cs
@@ -52,7 +52,7 @@ public class OnlineSubsystemTencent : ModuleRules
 			PublicDefinitions.Add("WITH_TENCENTSDK=0");
 		}
 
-		if (System.IO.Directory.Exists(System.IO.Path.Combine(Target.UEThirdPartySourceDirectory, "NotForLicensees", "Tencent")))
+		if (System.IO.Directory.Exists(System.IO.Path.Combine(EngineDirectory, "Restricted/NotForLicensees/Source/ThirdParty/Tencent")))
 		{
 			AddEngineThirdPartyPrivateDynamicDependencies(Target, "Tencent");
 			PublicIncludePathModuleNames.Add("Tencent");
diff --git a/Engine/Plugins/Runtime/Analytics/Apsalar/Source/IOSApsalar/IOSApsalar.Build.cs b/Engine/Plugins/Runtime/Analytics/Apsalar/Source/IOSApsalar/IOSApsalar.Build.cs
--- a/Engine/Plugins/Runtime/Analytics/Apsalar/Source/IOSApsalar/IOSApsalar.Build.cs
+++ b/Engine/Plugins/Runtime/Analytics/Apsalar/Source/IOSApsalar/IOSApsalar.Build.cs
@@ -1,5 +1,7 @@
 // Copyright Epic Games, Inc. All Rights Reserved.
 
+using System.IO;
+
 namespace UnrealBuildTool.Rules
 {
 	public class IOSApsalar : ModuleRules
@@ -47,16 +49,13 @@ namespace UnrealBuildTool.Rules
 					"z"
 				});
 
-            bool bHasApsalarSDK =
-                (System.IO.Directory.Exists(System.IO.Path.Combine(Target.UEThirdPartySourceDirectory, "Apsalar")) &&
-                  System.IO.Directory.Exists(System.IO.Path.Combine(Target.UEThirdPartySourceDirectory, "Apsalar", "IOS"))) ||
-                (System.IO.Directory.Exists(System.IO.Path.Combine(Target.UEThirdPartySourceDirectory, "NotForLicensees")) &&
-                  System.IO.Directory.Exists(System.IO.Path.Combine(Target.UEThirdPartySourceDirectory, "NotForLicensees", "Apsalar")) &&
-                  System.IO.Directory.Exists(System.IO.Path.Combine(Target.UEThirdPartySourceDirectory, "NotForLicensees", "Apsalar", "IOS")));
+			string SDKPath = Path.Combine(EngineDirectory, "Restricted/NotForLicensees/Source/ThirdParty/Apsalar/IOS");
+
+			bool bHasApsalarSDK = Directory.Exists(SDKPath);
             if (bHasApsalarSDK)
             {
-                PublicIncludePaths.Add(Target.UEThirdPartySourceDirectory + "NotForLicensees/Apsalar/IOS/");
-                PublicAdditionalLibraries.Add(Target.UEThirdPartySourceDirectory + "NotForLicensees/Apsalar/IOS/libApsalar.a");
+                PublicIncludePaths.Add(SDKPath);
+                PublicAdditionalLibraries.Add(Path.Combine(SDKPath, "libApsalar.a"));
 
                 PublicDefinitions.Add("WITH_APSALAR=1");
             }
diff --git a/Engine/Plugins/Runtime/PacketHandlers/CompressionComponents/Oodle/Source/OodleCompressionFormat/OodleCompressionFormat.Build.cs b/Engine/Plugins/Runtime/PacketHandlers/CompressionComponents/Oodle/Source/OodleCompressionFormat/OodleCompressionFormat.Build.cs
--- a/Engine/Plugins/Runtime/PacketHandlers/CompressionComponents/Oodle/Source/OodleCompressionFormat/OodleCompressionFormat.Build.cs
+++ b/Engine/Plugins/Runtime/PacketHandlers/CompressionComponents/Oodle/Source/OodleCompressionFormat/OodleCompressionFormat.Build.cs
@@ -30,8 +30,8 @@ public class OodleCompressionFormat : ModuleRules
 		}
 
 		// Check the NotForLicensees folder first
-		string OodleNotForLicenseesBuildFile = System.IO.Path.Combine(Target.UEThirdPartySourceDirectory, 
-			"NotForLicensees/RadGames/OodleData/OodleData.Build.cs");
+		string OodleNotForLicenseesBuildFile = System.IO.Path.Combine(EngineDirectory, 
+			"Restricted/NotForLicensees/Source/ThirdParty/RadGames/OodleData/OodleData.Build.cs");
 		string OodleLicenseesBuildFile = System.IO.Path.Combine(Target.UEThirdPartySourceDirectory,
 			"RadGames/OodleData/OodleData.Build.cs");
 
diff --git a/Engine/Plugins/Runtime/PacketHandlers/CompressionComponents/Oodle/Source/OodleHandlerComponent/OodleHandlerComponent.Build.cs b/Engine/Plugins/Runtime/PacketHandlers/CompressionComponents/Oodle/Source/OodleHandlerComponent/OodleHandlerComponent.Build.cs
--- a/Engine/Plugins/Runtime/PacketHandlers/CompressionComponents/Oodle/Source/OodleHandlerComponent/OodleHandlerComponent.Build.cs
+++ b/Engine/Plugins/Runtime/PacketHandlers/CompressionComponents/Oodle/Source/OodleHandlerComponent/OodleHandlerComponent.Build.cs
@@ -32,13 +32,13 @@ public class OodleHandlerComponent : ModuleRules
 			});
 		}
 
-		string OodleNotForLicenseesNetBuildFile = System.IO.Path.Combine(Target.UEThirdPartySourceDirectory,
-			"NotForLicensees/RadGames/OodleNet/OodleNet.Build.cs");
+		string OodleNotForLicenseesNetBuildFile = System.IO.Path.Combine(EngineDirectory,
+			"Restricted/NotForLicensees/Source/ThirdParty/RadGames/OodleNet/OodleNet.Build.cs");
 		string OodleLicenseesNetBuildFile = System.IO.Path.Combine(Target.UEThirdPartySourceDirectory,
 			"RadGames/OodleNet/OodleNet.Build.cs");
 
-		string OodleNotForLicenseesDataBuildFile = System.IO.Path.Combine(Target.UEThirdPartySourceDirectory,
-			"NotForLicensees/RadGames/OodleData/OodleData.Build.cs");
+		string OodleNotForLicenseesDataBuildFile = System.IO.Path.Combine(EngineDirectory,
+			"Restricted/NotForLicensees/Source/ThirdParty/RadGames/OodleData/OodleData.Build.cs");
 		string OodleLicenseesDataBuildFile = System.IO.Path.Combine(Target.UEThirdPartySourceDirectory,
 			"RadGames/OodleData/OodleData.Build.cs");
 
